Introduced screen reader and Slate screen reader plugins to Mac.

diff --git a/Engine/Plugins/Experimental/ScreenReader/ScreenReader.uplugin b/Engine/Plugins/Experimental/ScreenReader/ScreenReader.uplugin
--- a/Engine/Plugins/Experimental/ScreenReader/ScreenReader.uplugin
+++ b/Engine/Plugins/Experimental/ScreenReader/ScreenReader.uplugin
@@ -22,7 +22,8 @@
 			"LoadingPhase" : "PostEngineInit",
 			"WhitelistPlatforms" :
 			[
-				"Win64"
+				"Win64",
+				"Mac"
 			],
 			"BlacklistTargets" :
 			[
diff --git a/Engine/Plugins/Experimental/SlateScreenReader/SlateScreenReader.uplugin b/Engine/Plugins/Experimental/SlateScreenReader/SlateScreenReader.uplugin
--- a/Engine/Plugins/Experimental/SlateScreenReader/SlateScreenReader.uplugin
+++ b/Engine/Plugins/Experimental/SlateScreenReader/SlateScreenReader.uplugin
@@ -22,7 +22,8 @@
 			"LoadingPhase" : "PostEngineInit",
 			"WhitelistPlatforms" :
 			[
-				"Win64"
+				"Win64",
+				"Mac"
 			],
 			"BlacklistTargets" :
 			[
