Introduced screen reader and Slate screen reader plugins for Linux.

diff --git a/Engine/Plugins/Experimental/ScreenReader/ScreenReader.uplugin b/Engine/Plugins/Experimental/ScreenReader/ScreenReader.uplugin
--- a/Engine/Plugins/Experimental/ScreenReader/ScreenReader.uplugin
+++ b/Engine/Plugins/Experimental/ScreenReader/ScreenReader.uplugin
@@ -23,7 +23,9 @@
 			"WhitelistPlatforms" :
 			[
 				"Win64",
-				"Mac"
+				"Mac",
+				"Linux",
+				"LinuxArm64"
 			],
 			"BlacklistTargets" :
 			[
diff --git a/Engine/Plugins/Experimental/SlateScreenReader/SlateScreenReader.uplugin b/Engine/Plugins/Experimental/SlateScreenReader/SlateScreenReader.uplugin
--- a/Engine/Plugins/Experimental/SlateScreenReader/SlateScreenReader.uplugin
+++ b/Engine/Plugins/Experimental/SlateScreenReader/SlateScreenReader.uplugin
@@ -23,7 +23,9 @@
 			"WhitelistPlatforms" :
 			[
 				"Win64",
-				"Mac"
+				"Mac",
+				"Linux",
+				"LinuxArm64"
 			],
 			"BlacklistTargets" :
 			[
