Introduced BP API to the screen reader plugin.

diff --git a/Engine/Plugins/Experimental/ScreenReader/Source/ScreenReader/Private/Announcement/ScreenReaderAnnouncementChannel.cpp b/Engine/Plugins/Experimental/ScreenReader/Source/ScreenReader/Private/Announcement/ScreenReaderAnnouncementChannel.cpp
--- a/Engine/Plugins/Experimental/ScreenReader/Source/ScreenReader/Private/Announcement/ScreenReaderAnnouncementChannel.cpp
+++ b/Engine/Plugins/Experimental/ScreenReader/Source/ScreenReader/Private/Announcement/ScreenReaderAnnouncementChannel.cpp
@@ -98,6 +98,47 @@ FScreenReaderReply FScreenReaderAnnouncementChannel::RequestSpeakWidget(const TS
 	return RequestSpeak(FScreenReaderAnnouncement(WidgetAnnouncement, FScreenReaderAnnouncementInfo::DefaultWidgetAnnouncement()));
 }
 
+float FScreenReaderAnnouncementChannel::GetSpeechVolume() const
+{
+	return TextToSpeech->GetVolume();
+}
+
+FScreenReaderReply FScreenReaderAnnouncementChannel::SetSpeechVolume(float InVolume)
+{
+	// FTextToSpeech API already takes care of range check and clamping of values 
+	TextToSpeech->SetVolume(InVolume);
+	return FScreenReaderReply::Handled();
+}
+
+float FScreenReaderAnnouncementChannel::GetSpeechRate() const
+{
+	return TextToSpeech->GetRate();
+}
+
+FScreenReaderReply FScreenReaderAnnouncementChannel::SetSpeechRate(float InRate)
+{
+	// The FTextToSpeech API already handles range checks and clamping for values  
+	TextToSpeech->SetRate(InRate);
+	return FScreenReaderReply::Handled();
+}
+
+FScreenReaderReply FScreenReaderAnnouncementChannel::MuteSpeech()
+{
+	TextToSpeech->Mute();
+	return FScreenReaderReply::Handled();
+}
+
+FScreenReaderReply FScreenReaderAnnouncementChannel::UnmuteSpeech()
+{
+	TextToSpeech->Unmute();
+	return FScreenReaderReply::Handled();
+}
+
+bool FScreenReaderAnnouncementChannel::IsSpeechMuted() const
+{
+	return TextToSpeech->IsMuted();
+}
+
 void FScreenReaderAnnouncementChannel::OnTextToSpeechFinishSpeaking()
 {
 	// reset the current announcement to null. This is necessary for the RequestSpeak() algorithm 
diff --git a/Engine/Plugins/Experimental/ScreenReader/Source/ScreenReader/Private/GenericPlatform/ScreenReaderBase.cpp b/Engine/Plugins/Experimental/ScreenReader/Source/ScreenReader/Private/GenericPlatform/ScreenReaderBase.cpp
--- a/Engine/Plugins/Experimental/ScreenReader/Source/ScreenReader/Private/GenericPlatform/ScreenReaderBase.cpp
+++ b/Engine/Plugins/Experimental/ScreenReader/Source/ScreenReader/Private/GenericPlatform/ScreenReaderBase.cpp
@@ -51,11 +51,11 @@ bool FScreenReaderBase::RegisterUser(int32 InUserId)
 {
 	if (UsersMap.Contains(InUserId))
 	{
-		UE_LOG(LogScreenReader, Verbose, TEXT("Failed to register screen reader user with Id %i. Another user with the same Id has already been registered."), InUserId);
+		UE_LOG(LogScreenReader, Verbose, TEXT("Failed to register screen reader user with Id %d. Another user with the same Id has already been registered."), InUserId);
 		return false;
 	}
 	UsersMap.Add(InUserId, MakeShared<FScreenReaderUser>(InUserId));
-	UE_LOG(LogScreenReader, Verbose, TEXT("Registered screen reader user %i."), InUserId);
+	UE_LOG(LogScreenReader, Verbose, TEXT("Registered screen reader user %d."), InUserId);
 	return true;
 }
 
@@ -63,12 +63,12 @@ bool FScreenReaderBase::UnregisterUser(int32 InUserId)
 {
 	if (!UsersMap.Contains(InUserId))
 	{
-		UE_LOG(LogScreenReader, Verbose, TEXT("Failed to unregister with Id %i. No user with %i Id is registered."), InUserId);
+		UE_LOG(LogScreenReader, Verbose, TEXT("Failed to unregister with Id %d. No user with %d Id is registered."), InUserId, InUserId);
 		return false;
 	}
 	UsersMap[InUserId]->Deactivate();
 	UsersMap.Remove(InUserId);
-	UE_LOG(LogScreenReader, Verbose, TEXT("Unregistered screen reader user %i."), InUserId);
+	UE_LOG(LogScreenReader, Verbose, TEXT("Unregistered screen reader user %d."), InUserId);
 	return true;
 }
 
@@ -86,12 +86,13 @@ void FScreenReaderBase::UnregisterAllUsers()
 	UsersMap.Empty();
 }
 
-TSharedRef<FScreenReaderUser> FScreenReaderBase::GetUser(int32 InUserId) const
+TSharedRef<FScreenReaderUser> FScreenReaderBase::GetUserChecked(int32 InUserId) const
 {
+	checkf(IsUserRegistered(InUserId), TEXT("User Id %d is not registered. Did you forget to register the user with RegisterUser()?"), InUserId);
 	return UsersMap[InUserId];
 }
-
-TSharedPtr<FScreenReaderUser> FScreenReaderBase::GetUserChecked(int32 InUserId) const
+	
+TSharedPtr<FScreenReaderUser> FScreenReaderBase::GetUser(int32 InUserId) const
 {
 	if (IsUserRegistered(InUserId))
 	{
diff --git a/Engine/Plugins/Experimental/ScreenReader/Source/ScreenReader/Private/GenericPlatform/ScreenReaderUser.cpp b/Engine/Plugins/Experimental/ScreenReader/Source/ScreenReader/Private/GenericPlatform/ScreenReaderUser.cpp
--- a/Engine/Plugins/Experimental/ScreenReader/Source/ScreenReader/Private/GenericPlatform/ScreenReaderUser.cpp
+++ b/Engine/Plugins/Experimental/ScreenReader/Source/ScreenReader/Private/GenericPlatform/ScreenReaderUser.cpp
@@ -69,4 +69,48 @@ FScreenReaderReply FScreenReaderUser::RequestSpeakWidget(const TSharedRef<IAcces
 		return AnnouncementChannel->RequestSpeakWidget(InWidget);
 	}
 	return FScreenReaderReply::Unhandled();
+}
+
+float FScreenReaderUser::GetSpeechVolume() const
+{
+	return AnnouncementChannel->GetSpeechVolume();
+}
+
+FScreenReaderReply FScreenReaderUser::SetSpeechVolume(float InVolume)
+{
+	return AnnouncementChannel->SetSpeechVolume(InVolume);
+}
+
+float FScreenReaderUser::GetSpeechRate() const
+{
+	return AnnouncementChannel->GetSpeechRate();
+}
+
+FScreenReaderReply FScreenReaderUser::SetSpeechRate(float InRate)
+{
+	return AnnouncementChannel->SetSpeechRate(InRate);
+}
+
+FScreenReaderReply FScreenReaderUser::MuteSpeech()
+{
+	if (IsActive())
+	{
+
+	}
+	return FScreenReaderReply::Unhandled();
+}
+
+FScreenReaderReply FScreenReaderUser::UnmuteSpeech()
+{
+	if (IsActive())
+	{
+		return AnnouncementChannel->UnmuteSpeech();
+	}
+
+	return FScreenReaderReply::Unhandled();
+}
+
+bool FScreenReaderUser::IsSpeechMuted() const
+{
+	return AnnouncementChannel->IsSpeechMuted();
 }
\ No newline at end of file
diff --git a/Engine/Plugins/Experimental/ScreenReader/Source/ScreenReader/Public/Announcement/ScreenReaderAnnouncement.h b/Engine/Plugins/Experimental/ScreenReader/Source/ScreenReader/Public/Announcement/ScreenReaderAnnouncement.h
--- a/Engine/Plugins/Experimental/ScreenReader/Source/ScreenReader/Public/Announcement/ScreenReaderAnnouncement.h
+++ b/Engine/Plugins/Experimental/ScreenReader/Source/ScreenReader/Public/Announcement/ScreenReaderAnnouncement.h
@@ -35,10 +35,22 @@ public:
 	// Please use ctor that takes parameters 
 	FScreenReaderAnnouncementInfo() = default;
 	FScreenReaderAnnouncementInfo(bool bInShouldQueue, bool bInInterruptable, EScreenReaderAnnouncementPriority InPriority);
-	bool ShouldQueue() const { return bShouldQueue; }
-	bool IsInterruptable() const { return bInterruptable; }
-	EScreenReaderAnnouncementPriority GetPriority() const { return Priority; }
-	double GetTimestamp() const { return Timestamp; }
+	bool ShouldQueue() const 
+	{ 
+		return bShouldQueue; 
+	}
+	bool IsInterruptable() const 
+	{ 
+		return bInterruptable; 
+	}
+	EScreenReaderAnnouncementPriority GetPriority() const 
+	{ 
+		return Priority; 
+	}
+	double GetTimestamp() const 
+	{ 
+		return Timestamp; 
+	}
 	/** Returns a string that represents all of the data reflected in this struct. */
 	FString ToString() const;
 
@@ -76,11 +88,14 @@ public:
 	}
 
 private:
-	/** True if the associazted announcement should be queued if it cannot be spoken immediately or is interrupted. Else false.*/
+	/** True if the associated announcement should be queued if it cannot be spoken immediately or is interrupted. Else false.*/
+	UPROPERTY(BlueprintReadWrite, Category="ScreenReaderAnnouncementInfo", meta = (AllowPrivateAccess = "true"))
 	bool bShouldQueue = false;
 	/** True if the associazted announcement can be intrrupted by another announcement. Else false. */
+	UPROPERTY(BlueprintReadWrite, Category="ScreenReaderAnnouncementInfo", meta = (AllowPrivateAccess = "true"))
 	bool bInterruptable= true;
 	/** The priority level of the associated announcement. */
+	UPROPERTY(BlueprintReadWrite, Category="ScreenReaderAnnouncementInfo", meta = (AllowPrivateAccess = "true"))
 	EScreenReaderAnnouncementPriority Priority = EScreenReaderAnnouncementPriority::Medium;
 	// @TODOAccessibility: Update timestamp to be recorded when announcements are requested, not when they are constructed. 
 	/** The timestamp for when the associated announcement was created. */
@@ -161,11 +176,11 @@ public:
 	{
 		return InInterruptionPolicy.CanBeInterruptedBy(Interruptee, Interrupter);
 	}
-private:
+public:
 	/** A localized string that represents the message to be spoken to a end user */
+	UPROPERTY(BlueprintReadWrite, Category="ScreenReaderAnnouncement")
 	FString AnnouncementString;
-	/** The announcement info assocated with the announcement which controls how the announcement behaves when a user requests this announcement to be spoken to an end user */
+	/** The announcement info associated with the announcement which controls how the announcement behaves when a user requests this announcement to be spoken to an end user */
+	UPROPERTY(BlueprintReadWrite, Category="ScreenReaderAnnouncement")
 	FScreenReaderAnnouncementInfo AnnouncementInfo;
 };
-
-
diff --git a/Engine/Plugins/Experimental/ScreenReader/Source/ScreenReader/Public/Announcement/ScreenReaderAnnouncementChannel.h b/Engine/Plugins/Experimental/ScreenReader/Source/ScreenReader/Public/Announcement/ScreenReaderAnnouncementChannel.h
--- a/Engine/Plugins/Experimental/ScreenReader/Source/ScreenReader/Public/Announcement/ScreenReaderAnnouncementChannel.h
+++ b/Engine/Plugins/Experimental/ScreenReader/Source/ScreenReader/Public/Announcement/ScreenReaderAnnouncementChannel.h
@@ -32,7 +32,13 @@ public:
 	FScreenReaderReply ClearAllAnnouncements();
 	FScreenReaderReply RequestSpeakWidget(const TSharedRef<IAccessibleWidget>& InWidget);
 
-	
+	float GetSpeechVolume() const;
+	FScreenReaderReply SetSpeechVolume(float InVolume);
+	float GetSpeechRate() const;
+	FScreenReaderReply SetSpeechRate(float InRate);
+	FScreenReaderReply MuteSpeech();
+	FScreenReaderReply UnmuteSpeech();
+	bool IsSpeechMuted() const;
 private:
 	/**
 	* A callback that is bound to the text to speech delegate. The delegate fires when a string
diff --git a/Engine/Plugins/Experimental/ScreenReader/Source/ScreenReader/Public/GenericPlatform/ScreenReaderBase.h b/Engine/Plugins/Experimental/ScreenReader/Source/ScreenReader/Public/GenericPlatform/ScreenReaderBase.h
--- a/Engine/Plugins/Experimental/ScreenReader/Source/ScreenReader/Public/GenericPlatform/ScreenReaderBase.h
+++ b/Engine/Plugins/Experimental/ScreenReader/Source/ScreenReader/Public/GenericPlatform/ScreenReaderBase.h
@@ -67,18 +67,18 @@ public:
 	/** Unregisteres all screen reader users. No users can receive or respond to accessible events or accessible input upon calling this function. */
 	void UnregisterAllUsers();
 	/** 
-	* Returns the screen reader user associated with the passed in Id. It is up to the user
-	* to ensure that the Id is valid and mapped to a valid screen reader user.
+	* Returns the screen reader user associated with the passed in Id. An assertion will trigger if the User Id is not registered. 
+	* It is up to the user to ensure that the Id is valid and mapped to a valid screen reader user.
 	* @param InUserId The Id of the screen reader user to be retrieved.
 	* @return The screen reader user associated with the passed in Id
 	*/
-	TSharedRef<FScreenReaderUser> GetUser(int32 InUserId) const;
+	TSharedRef<FScreenReaderUser> GetUserChecked(int32 InUserId) const;
 	/**
 	* Returns a screen reader user associated with the passed in Id if the Id is registered. Otherwise, a nullptr is returned.
 	* @param InUserId The user Id of the screen reader to be retrieved.
-	* @return The screen reader user associated with the Id if the user is registered. Else false.
+	* @return The screen reader user associated with the Id if the user is registered. Else nullptr.
 	*/
-	TSharedPtr<FScreenReaderUser> GetUserChecked(int32 InUserId) const;
+	TSharedPtr<FScreenReaderUser> GetUser(int32 InUserId) const;
 
 protected:
 	/**
diff --git a/Engine/Plugins/Experimental/ScreenReader/Source/ScreenReader/Public/GenericPlatform/ScreenReaderReply.h b/Engine/Plugins/Experimental/ScreenReader/Source/ScreenReader/Public/GenericPlatform/ScreenReaderReply.h
--- a/Engine/Plugins/Experimental/ScreenReader/Source/ScreenReader/Public/GenericPlatform/ScreenReaderReply.h
+++ b/Engine/Plugins/Experimental/ScreenReader/Source/ScreenReader/Public/GenericPlatform/ScreenReaderReply.h
@@ -2,16 +2,19 @@
 
 #pragma once
 #include "CoreMinimal.h"
+#include "ScreenReaderReply.generated.h"
 /**
 * A struct passed around the the screen reader framework to indicate if a request has been successfully handled.
 * Use this class to determine if a request to the screen reader user was successfully handled and provide user feedback accordingly.
 * // @TODOAccessibility: Provide examples 
 * @see FScreenReaderUser
 */
-class SCREENREADER_API FScreenReaderReply
+USTRUCT(BlueprintType)
+struct SCREENREADER_API FScreenReaderReply
 {
+	GENERATED_BODY()
 public:
-	FScreenReaderReply() = delete;
+	FScreenReaderReply() = default;
 	~FScreenReaderReply() = default;
 
 	static FScreenReaderReply Handled()
@@ -34,5 +37,6 @@ private:
 	{
 
 	}
-	bool bHandled;
+	UPROPERTY(BlueprintReadOnly, Category=ScreenReaderReply, meta=(AllowPrivateAccess = "true"))
+	bool bHandled = false;
 };
\ No newline at end of file
diff --git a/Engine/Plugins/Experimental/ScreenReader/Source/ScreenReader/Public/GenericPlatform/ScreenReaderUser.h b/Engine/Plugins/Experimental/ScreenReader/Source/ScreenReader/Public/GenericPlatform/ScreenReaderUser.h
--- a/Engine/Plugins/Experimental/ScreenReader/Source/ScreenReader/Public/GenericPlatform/ScreenReaderUser.h
+++ b/Engine/Plugins/Experimental/ScreenReader/Source/ScreenReader/Public/GenericPlatform/ScreenReaderUser.h
@@ -62,31 +62,67 @@ public:
 	* @return FScreenReaderReply::Handled() if the widget's information is successfully spoken. Else returns FScreenReaderReply::Unhandled();
 	*/
 	FScreenReaderReply RequestSpeakWidget(const TSharedRef<IAccessibleWidget>& InWidget);
+	/** Returns the speech volume for text to speech. Value will be between 0.0f and 1.0f. */
+	float GetSpeechVolume() const;
+	/**
+	* Sets the speech volume for text to speech to speak at. The passed in value will be clamped between 0.0f and 1.0f.
+	* @return FScreenReaderReply::Handled() if the volume is successfully set. Else FScreenReader::Unhandled() is returned.
+	*/
+	FScreenReaderReply SetSpeechVolume(float InVolume);
+	/** Returns the speech rate text to speech is speaking at for this user. Value will be between 0.0f and 1.0f. */
+	float GetSpeechRate() const;
+	/**
+	* Sets the speech rate text to speech will speak at for this user. Passed in value is clamped between 0.0f and 1.0f.
+	* @return FScreenReaderReply::Handled() if the speech rate is successfully set. Else, FScreenReaderReply::Unhandled() is returned.
+	*/
+	FScreenReaderReply SetSpeechRate(float InRate);
+	/**
+	* Mutes the text to speech for this user so requests to speak strings will be inaudible.
+	* @return FScreenReaderReply::Handled() if the text to speech is successfully muted. Else, FScreenReaderReply::Unhandled() is returned.
+	*/
+	FScreenReaderReply MuteSpeech();
+	/**
+	* Unmutes the text to speech for this user so that requests to speak strings will be audible again.
+	* @return FScreenReaderReply::Handled() if the text to speech si successfully unmuted. Else, FScreenReaderReply::Unhandled() is returned.
+	*/
+	FScreenReaderReply UnmuteSpeech();
+	/** Returns true if the text to speech for this user is muted. Else returns false. */
+	bool IsSpeechMuted() const;
 
 	/**
-	* Activates the screen reader and allow users to use its services such as requesting announcements to be spoken via text to speech.
-	* When screen reader users are firstr registered with a screen reader, they are inactive by default. Users must explicitly activate the screen reader user.
+	* Activates the screen reader user and fulfill requests for accessibility services such as text to speech that clients can make.
+	* When screen reader users are firstr registered with a screen reader, they are deactivated by default. Users must explicitly activate the screen reader user.
 	*/
 	void Activate();
 	/** 
 	* Deactivates the screen reader and disables all announcement and text to speech services
 	* making them do nothing.
 	*/
-
 	void Deactivate();
 	/** Returns true if the screen reader user is active. Else returns false.*/
-	bool IsActive() const { return bActive; }
-
-	/** Retursn the user Id associated with this screen reader user. */
-	int32 GetUserId() const { return UserId; }
+	bool IsActive() const 
+	{ 
+		return bActive; 
+	}
+	/** Returns the user Id associated with this screen reader user. */
+	int32 GetUserId() const 
+	{ 
+		return UserId; 
+	}
 	/** Returns the accessible widget the screen reader user is currently focused on. */
-	TSharedPtr<IAccessibleWidget> GetAccessibleFocus() const { return AccessibleFocus.Pin(); }
-	/** Sets the accessible the screenr eader user is currently focused on */
-	void SetAccessibleFocus(const TSharedRef<IAccessibleWidget>& InAccessibleFocus) { AccessibleFocus = InAccessibleFocus; }
+	TSharedPtr<IAccessibleWidget> GetAccessibleFocusWidget() const 
+	{ 
+		return AccessibleFocusWidget.Pin();
+	}
+	/** Sets the accessible widget the screen reader user is currently focused on */
+	void SetAccessibleFocusWidget(const TSharedRef<IAccessibleWidget>& InAccessibleFocusWidget) 
+	{ 
+		AccessibleFocusWidget = InAccessibleFocusWidget;
+	}
 private:
 	int32 UserId;
 	/** Responsible for handling all incoming announcement requests and speaking them via text to speech if possible */
 	TUniquePtr<FScreenReaderAnnouncementChannel> AnnouncementChannel;
-	TWeakPtr<IAccessibleWidget> AccessibleFocus;
+	TWeakPtr<IAccessibleWidget> AccessibleFocusWidget;
 	bool bActive;
-};
\ No newline at end of file
+};
diff --git a/Engine/Plugins/Experimental/SlateScreenReader/Source/SlateScreenReader/Private/SlateScreenReader.cpp b/Engine/Plugins/Experimental/SlateScreenReader/Source/SlateScreenReader/Private/SlateScreenReader.cpp
--- a/Engine/Plugins/Experimental/SlateScreenReader/Source/SlateScreenReader/Private/SlateScreenReader.cpp
+++ b/Engine/Plugins/Experimental/SlateScreenReader/Source/SlateScreenReader/Private/SlateScreenReader.cpp
@@ -32,12 +32,12 @@ void FSlateScreenReader::HandleSlateFocusChanging(const FFocusEvent& InFocusEven
 	if (IsActive() && InNewFocusWidget && InNewFocusWidget->IsAccessible())
 	{
 		int32 UserId = static_cast<int32>(InFocusEvent.GetUser());
-		if (TSharedPtr<FScreenReaderUser> User = GetUserChecked(UserId))
+		if (TSharedPtr<FScreenReaderUser> User = GetUser(UserId))
 		{
 			TSharedPtr<IAccessibleWidget> AccessibleWidget = FSlateAccessibleWidgetCache::GetAccessibleWidgetChecked(InNewFocusWidget);
 			if (AccessibleWidget)
 			{
-				User->SetAccessibleFocus(AccessibleWidget.ToSharedRef());
+				User->SetAccessibleFocusWidget(AccessibleWidget.ToSharedRef());
 				User->RequestSpeakWidget(AccessibleWidget.ToSharedRef());
 			}
 		}
diff --git a/Engine/Plugins/Experimental/SlateScreenReader/Source/SlateScreenReader/Private/SlateScreenReaderEngineSubsystem.cpp b/Engine/Plugins/Experimental/SlateScreenReader/Source/SlateScreenReader/Private/SlateScreenReaderEngineSubsystem.cpp
--- a/Engine/Plugins/Experimental/SlateScreenReader/Source/SlateScreenReader/Private/SlateScreenReaderEngineSubsystem.cpp
+++ b/Engine/Plugins/Experimental/SlateScreenReader/Source/SlateScreenReader/Private/SlateScreenReaderEngineSubsystem.cpp
@@ -1,7 +1,10 @@
 // Copyright Epic Games, Inc. All Rights Reserved.
 
 #include "SlateScreenReaderEngineSubsystem.h"
+#include "GenericPlatform/ScreenReaderReply.h"
 #include "GenericPlatform/ScreenReaderBase.h"
+#include "GenericPlatform/ScreenReaderUser.h"
+#include "SlateScreenReaderLog.h"
 #include "Application/SlateApplicationBase.h"
 #include "GenericPlatform/IScreenReaderBuilder.h"
 #include "SlateScreenReaderModule.h"
@@ -62,15 +65,16 @@ bool USlateScreenReaderEngineSubsystem::ShouldCreateSubsystem(UObject* Outer) co
 	return FSlateApplication::IsInitialized();
 }
 
-void USlateScreenReaderEngineSubsystem::Activate()
+void USlateScreenReaderEngineSubsystem::ActivateScreenReader()
 {
+	checkf(ScreenReader, TEXT("Screen reader is invalid during activation. A screen reader should have been created on subsystem initialization and be valid for the duration of the subsystem."));
 	if (!ScreenReader->IsActive())
 	{
 		ScreenReader->Activate();
 	}
 }
 
-void USlateScreenReaderEngineSubsystem::Deactivate()
+void USlateScreenReaderEngineSubsystem::DeactivateScreenReader()
 {
 	checkf(ScreenReader, TEXT("Invalid screen reader on deactivation. A valid screen reader must always be available for deactivation."));
 	if (ScreenReader->IsActive())
@@ -79,7 +83,207 @@ void USlateScreenReaderEngineSubsystem::Deactivate()
 	}
 }
 
+bool USlateScreenReaderEngineSubsystem::IsScreenReaderActive() const
+{
+	checkf(ScreenReader, TEXT("Screen reader should have been created during subsystem initialization and should be valid for the entire lifetime of the engine subsystem."));
+	return ScreenReader->IsActive();
+}
+
+FScreenReaderReply USlateScreenReaderEngineSubsystem::ActivateUser(UPARAM(DisplayName = "Screen Reader User Id") int32 InUserId)
+{
+	checkf(ScreenReader, TEXT("Screen reader should have been created during subsystem initialization and should be valid for the entire lifetime of the engine subsystem."));
+	TSharedPtr<FScreenReaderUser> User = ScreenReader->GetUser(InUserId);
+	if (User)
+	{
+		User->Activate();
+		return FScreenReaderReply::Handled();
+	}
+	UE_LOG(LogSlateScreenReader, Warning, TEXT("No screen reader user with %d Id exists. Did you forget to register the user Id?"), InUserId);
+	return FScreenReaderReply::Unhandled();
+}
+
+FScreenReaderReply USlateScreenReaderEngineSubsystem::DeactivateUser(UPARAM(DisplayName = "Screen Reader User Id") int32 InUserId)
+{
+	checkf(ScreenReader, TEXT("Screen reader should have been created during subsystem initialization and should be valid for the entire lifetime of the engine subsystem."));
+	TSharedPtr<FScreenReaderUser> User = ScreenReader->GetUser(InUserId);
+	if (User)
+	{
+		// FScreenReaderUser already checks if it is active/inactive 
+		User->Deactivate();
+		return FScreenReaderReply::Handled();
+	}
+	UE_LOG(LogSlateScreenReader, Warning, TEXT("No screen reader user with %d Id exists. Did you forget to register the user Id?"), InUserId);
+	return FScreenReaderReply::Unhandled();
+}
+
+FScreenReaderReply USlateScreenReaderEngineSubsystem::RegisterUser(UPARAM(DisplayName = "Screen Reader User Id") int32 InUserId)
+{
+	return ScreenReader->RegisterUser(InUserId)
+		? FScreenReaderReply::Handled()
+		: FScreenReaderReply::Unhandled();
+}
+
+FScreenReaderReply USlateScreenReaderEngineSubsystem::UnregisterUser(UPARAM(DisplayName = "Screen Reader User Id") int32 InUserId)
+{
+	checkf(ScreenReader, TEXT("Screen reader should have been created during subsystem initialization and should be valid for the entire lifetime of the engine subsystem."));
+	return ScreenReader->UnregisterUser(InUserId)
+		? FScreenReaderReply::Handled()
+		: FScreenReaderReply::Unhandled();
+}
+
+bool USlateScreenReaderEngineSubsystem::IsUserRegistered(UPARAM(DisplayName = "Screen Reader User Id") int32 InUserId) const
+{
+	checkf(ScreenReader, TEXT("Screen reader should have been created during subsystem initialization and should be valid for the entire lifetime of the engine subsystem."));
+	return ScreenReader->IsUserRegistered(InUserId);
+}
+
+FScreenReaderReply USlateScreenReaderEngineSubsystem::RequestSpeak(UPARAM(DisplayName = "Screen Reader User Id") int32 InUserId, UPARAM(DisplayName = "Announcement") FScreenReaderAnnouncement InAnnouncement)
+{
+	checkf(ScreenReader, TEXT("Screen reader should have been created during subsystem initialization and should be valid for the entire lifetime of the engine subsystem."));
+	if (ScreenReader->IsActive())
+	{
+		TSharedPtr<FScreenReaderUser> User = ScreenReader->GetUser(InUserId);
+		if (User)
+		{
+			return User->RequestSpeak(MoveTemp(InAnnouncement));
+		}
+	}
+	return FScreenReaderReply::Unhandled();
+}
+
+FScreenReaderReply USlateScreenReaderEngineSubsystem::RequestSpeakFocusedWidget(UPARAM(DisplayName = "Screen Reader User Id") int32 InUserId)
+{
+	checkf(ScreenReader, TEXT("Screen reader should have been created during subsystem initialization and should be valid for the entire lifetime of the engine subsystem."));
+	if (ScreenReader->IsActive())
+	{
+		TSharedPtr<FScreenReaderUser> User = ScreenReader->GetUser(InUserId);
+		if (User)
+		{
+			if (TSharedPtr<IAccessibleWidget> UserFocusedWidget = User->GetAccessibleFocusWidget())
+			{
+				return User->RequestSpeakWidget(UserFocusedWidget.ToSharedRef());
+			}
+		}
+	}
+	return FScreenReaderReply::Unhandled();
+}
+
+FScreenReaderReply USlateScreenReaderEngineSubsystem::StopSpeaking(UPARAM(DisplayName = "Screen Reader User Id") int32 InUserId)
+{
+	checkf(ScreenReader, TEXT("Screen reader should have been created during subsystem initialization and should be valid for the entire lifetime of the engine subsystem."));
+	if (IsScreenReaderActive())
+	{
+		TSharedPtr<FScreenReaderUser> User = ScreenReader->GetUser(InUserId);
+		if (User)
+		{
+			User->StopSpeaking();
+			return FScreenReaderReply::Handled();
+		}
+	}
+	return FScreenReaderReply::Unhandled();
+}
+
+bool USlateScreenReaderEngineSubsystem::IsSpeaking(UPARAM(DisplayName = "Screen Reader User Id") int32 InUserId) const
+{
+	checkf(ScreenReader, TEXT("Screen reader should have been created during subsystem initialization and should be valid for the entire lifetime of the engine subsystem."));
+	if (ScreenReader->IsActive())
+	{
+		TSharedPtr<FScreenReaderUser> User = ScreenReader->GetUser(InUserId);
+		if (User)
+		{
+			return User->IsSpeaking();
+		}
+	}
+	return false;
+}
+
+float USlateScreenReaderEngineSubsystem::GetSpeechVolume(UPARAM(DisplayName = "Screen Reader User Id") UPARAM(DisplayName = "Screen Reader User Id") int32 InUserId) const
+{
+	checkf(ScreenReader, TEXT("Screen reader should have been created during subsystem initialization and should be valid for the entire lifetime of the engine subsystem."));
+	TSharedPtr<FScreenReaderUser> User = ScreenReader->GetUser(InUserId);
+	if (User)
+	{
+		return User->GetSpeechVolume();
+	}
+	return 0.0f;
+}
+
+FScreenReaderReply USlateScreenReaderEngineSubsystem::SetSpeechVolume(UPARAM(DisplayName = "Screen Reader User Id") UPARAM(DisplayName = "Screen Reader User Id") int32 InUserId, UPARAM(DisplayName = "Volume") float InVolume)
+{
+	checkf(ScreenReader, TEXT("Screen reader should have been created during subsystem initialization and should be valid for the entire lifetime of the engine subsystem."));
+	TSharedPtr<FScreenReaderUser> User = ScreenReader->GetUser(InUserId);
+	if (User)
+	{
+		return User->SetSpeechVolume(InVolume);
+	}
+	return FScreenReaderReply::Unhandled();
+}
+
+float USlateScreenReaderEngineSubsystem::GetSpeechRate(UPARAM(DisplayName = "Screen Reader User Id") UPARAM(DisplayName = "Screen Reader User Id") int32 InUserId) const
+{
+	checkf(ScreenReader, TEXT("Screen reader should have been created during subsystem initialization and should be valid for the entire lifetime of the engine subsystem."));
+	TSharedPtr<FScreenReaderUser> User = ScreenReader->GetUser(InUserId);
+	if (User)
+	{
+		return User->GetSpeechRate();
+	}
+	return 0.0f;
+}
+
+FScreenReaderReply USlateScreenReaderEngineSubsystem::SetSpeechRate(UPARAM(DisplayName = "Screen Reader User Id") UPARAM(DisplayName = "Screen Reader User Id") int32 InUserId, UPARAM(DisplayName = "Rate") float InRate)
+{
+	checkf(ScreenReader, TEXT("Screen reader should have been created during subsystem initialization and should be valid for the entire lifetime of the engine subsystem."));
+	TSharedPtr<FScreenReaderUser> User = ScreenReader->GetUser(InUserId);
+	if (User)
+	{
+		return User->SetSpeechRate(InRate);
+	}
+	return FScreenReaderReply::Unhandled();
+}
+
+FScreenReaderReply USlateScreenReaderEngineSubsystem::MuteSpeech(UPARAM(DisplayName = "Screen Reader User Id") UPARAM(DisplayName = "Screen Reader User Id") int32 InUserId)
+{
+	checkf(ScreenReader, TEXT("Screen reader should have been created during subsystem initialization and should be valid for the entire lifetime of the engine subsystem."));
+	if (IsScreenReaderActive())
+	{
+		TSharedPtr<FScreenReaderUser> User = ScreenReader->GetUser(InUserId);
+		if (User)
+		{
+			return User->MuteSpeech();
+		}
+	}
+	return FScreenReaderReply::Unhandled();
+}
+
+FScreenReaderReply USlateScreenReaderEngineSubsystem::UnmuteSpeech(UPARAM(DisplayName = "Screen Reader User Id") UPARAM(DisplayName = "Screen Reader User Id") int32 InUserId)
+{
+	checkf(ScreenReader, TEXT("Screen reader should have been created during subsystem initialization and should be valid for the entire lifetime of the engine subsystem."));
+	if (IsScreenReaderActive())
+	{
+		TSharedPtr<FScreenReaderUser> User = ScreenReader->GetUser(InUserId);
+		if (User)
+		{
+			return User->UnmuteSpeech();
+		}
+	}
+	return FScreenReaderReply::Unhandled();
+}
+
+bool USlateScreenReaderEngineSubsystem::IsSpeechMuted(UPARAM(DisplayName = "Screen Reader User Id") UPARAM(DisplayName = "Screen Reader User Id") int32 InUserId) const
+{
+	checkf(ScreenReader, TEXT("Screen reader should have been created during subsystem initialization and should be valid for the entire lifetime of the engine subsystem."));
+	if (IsScreenReaderActive())
+	{
+		TSharedPtr<FScreenReaderUser> User = ScreenReader->GetUser(InUserId);
+		if (User)
+		{
+			return User->IsSpeechMuted();
+		}
+	}
+	return false;
+}
+
 TSharedRef<FScreenReaderBase> USlateScreenReaderEngineSubsystem::GetScreenReader() const
 {
 	return ScreenReader.ToSharedRef();
-}
\ No newline at end of file
+}
diff --git a/Engine/Plugins/Experimental/SlateScreenReader/Source/SlateScreenReader/Public/SlateScreenReaderEngineSubsystem.h b/Engine/Plugins/Experimental/SlateScreenReader/Source/SlateScreenReader/Public/SlateScreenReaderEngineSubsystem.h
--- a/Engine/Plugins/Experimental/SlateScreenReader/Source/SlateScreenReader/Public/SlateScreenReaderEngineSubsystem.h
+++ b/Engine/Plugins/Experimental/SlateScreenReader/Source/SlateScreenReader/Public/SlateScreenReaderEngineSubsystem.h
@@ -3,19 +3,24 @@
 #pragma once
 #include "CoreMinimal.h"
 #include "Subsystems/EngineSubsystem.h"
+#include "Announcement/ScreenReaderAnnouncement.h"
+#include "GenericPlatform/ScreenReaderReply.h"
 #include "SlateScreenReaderEngineSubsystem.generated.h"
 
 class FScreenReaderBase;
 /**
-* The engine susbsystem for the Slate screen reader.
+* The engine subsystem for the Slate screen reader.
+* A screen reader is a framework that provides vision accessibility services for screen reader useres.
+* A screen reader user is a single user of the screen reader framework and can be thought of as a user of a hardware device such as keyboard/mouse or a controller.
+* Screen reader users must be registered with the screen reader framework for them to receive feedback from the accessibility services provided by the screen reader framework.
 * This class should be the entryway for C++ programmers and BP useres alike to interact with the screen reader system.
 * The subsystem must be activated before the screen reader services can be used.
 * For C++ users, please retrieve the screen reader and interact with the screen reader users from there.
 * Example:
-* USlateScreenReaderEngineSubsystem ::Get().Activate();
+* USlateScreenReaderEngineSubsystem ::Get().ActivateScreenReader();
 * // Registers a screen reader user with Id 0. A screen reader user should correspond to a hardware input device such as a keyboard or controller like FSlateUser
 * USlateScreenReaderEngineSubsystem ::Get().GetScreenReader()->RegisterUser(0);
-* TSharedRef<FScreenReaderUser> User = USlateScreenReaderEngineSubsystem::Get().GetScreenReader()->GetUser(0);
+* TSharedRef<FScreenReaderUser> User = USlateScreenReaderEngineSubsystem::Get().GetScreenReader()->GetUserChecked(0);
 * // Screen reader users are inactive when they are first registered and need to be explicitly activated.
 * User->Activate();
 * static const FText HelloWorld = LOCTEXT("HelloWorld", "Hello World");
@@ -35,17 +40,184 @@ public:
 	static USlateScreenReaderEngineSubsystem& Get();
 	
 	/** 
-	* Activates the underlying screen reader. Use this to allow end users to register with the screen reader 
+	* Activates the underlying screen reader. Use this to allow screen reader users to register with the screen reader 
 	* and receive accessible feedback via text to speech and get access to other screen reader services.
+	* A basic workflow with activation would be:
+	* ActivateScreenReader() -> RegisterScreenReaderUser() -> ActivateScreenReaderUser()
+	* @see RegisterScreenReaderUser(), ActivateScreenReaderUser()
 	*/
-	UFUNCTION(BlueprintCallable, Category="ScreenReader")
-	void Activate();
+	UFUNCTION(BlueprintCallable, Category="SlateScreenReader")
+	void ActivateScreenReader();
+
 	/**
-	* Deactivates the underlying screen reader and prevents end users from getting
+	* Deactivates the underlying screen reader and prevents screen reader users from getting
 	* any accessible feedback via text to speech or using any other screen reader services.
+	* Note: When the screen reader is deactivated, none of the registered screen reader users will be unregistered or cleared. This allows you to 
+	* deactivate the screen reader to prevent accessible services such as text to speech from triggering and then activate the screen reader again to continue those services.
+	*/
+	UFUNCTION(BlueprintCallable, Category="SlateScreenReader")
+	void DeactivateScreenReader();
+
+	/**
+	* Returns true if the screen reader is currently active and accessibility services such as text to speech can be used by the screen reader users. 
+	* Otherwise, it returns false.
+	*/
+	UFUNCTION(BlueprintCallable, Category = "SlateScreenReader")
+	bool IsScreenReaderActive() const;
+
+	/**
+	* Activates a screen reader user and fulfill requests for accessibility services such as text to speech that clients can make.
+	* When screen reader users are first registered with a screen reader, they are deactivated by default. Users must explicitly activate the screen reader user.
+	* If the passed in user Id does not correspond to a registered screen reader user, nothing will happen.
+	* @param InUserId The user Id of the screen reader user to request activation 
+	* @return FScreenReaderReply::Handled() if the screen reader user is successfully activated. Else FScreenReaderReply::Unhandled() is returned.
+	*/
+	UFUNCTION(BlueprintCallable, Category = "SlateScreenReader")
+	FScreenReaderReply ActivateUser(UPARAM(DisplayName = "Screen Reader User Id") int32 InUserId);
+
+	/**
+	* Deactivates the screen reader and disables all announcement and text to speech services
+	* making them do nothing.
+	* If the passed in user Id does not correspond to a registered screen reader user, nothing will happen.
+	* @param InUserId The user Id of the screen reader user to request deactivation. The Id should correspond to the Slate user id of an input device. If unsure, use Id 0. 
+	* @return FScreenReaderReply::Handled() if the screen reader user is successfully deactivated. Else FScreenReaderReply::Unhandled() is returned.
+	*/
+	UFUNCTION(BlueprintCallable, Category = "SlateScreenReader")
+	FScreenReaderReply DeactivateUser(UPARAM(DisplayName = "Screen Reader User Id") int32 InUserId);
+
+	/**
+	* Registers a provided user Id to the screen reader framework and allows the screen reader user to receive and respond to accessible events and accessible input.
+	* Does nothing if the passed in Id is already registered.
+	* Note: A successfully registered screen reader user is deactivated by default and will not respond to accessible events or accessible input.
+	* You need to call ActivateUser() to allow the newly registered screen reader user to respond to the accessible events and accessible input.
+	* @param InUserId The user Id of the screen reader user to register. The Id should correspond to a valid Slate user Id for an active hardware input device. If unsure, use Id 0.
+	* @return FScreenReaderReply::Handled() if the screen reader user is successfully registered. Else FScreenReaderReply::Unhandled() is returned.
+	* @see ActivateUser()
+	*/
+	UFUNCTION(BlueprintCallable, Category = "SlateScreenReader")
+	FScreenReaderReply RegisterUser(UPARAM(DisplayName = "Screen Reader User Id") int32 InUserId);
+
+	/**
+	* Unregisters a provided user Id from the screen reader framework and deactivates the user. The unregistered user will no longer receive or respond to accessible events and input.
+	* Nothing will happen if the provided user Id has not been registered with the screen reader.
+	* @param InUserId The user Id to unregister from the screen reader framework.
+	* @return FScreenReaderReply::Handled() if the screen reader user is successfully unregistered. Else FScreenReaderReply::Unhandled() is returned.
+	*/
+	UFUNCTION(BlueprintCallable, Category = "SlateScreenReader")
+	FScreenReaderReply UnregisterUser(UPARAM(DisplayName = "Screen Reader User Id") int32 InUserId);
+	/**
+	* Returns true if the passed in screen reader user Id is already registered. Else false is returned.
+	* @param InUserId The user Id to check for registration.
+	*/
+	UFUNCTION(BlueprintCallable, Category = "SlateScreenReader")
+	bool IsUserRegistered(UPARAM(DisplayName = "Screen Reader User Id") int32 InUserId) const;
+
+	/**
+	* Requests an announcement to be spoken to the screen reader user. This is the main mechanism to provide text to speech auditory feedback
+	* to end users.
+	*Calling this function does not guarantee that the announcement will be spoken via text to speech and be
+	* heard by a user.
+	* All announcements spoken via text to speech will be asynchronous and will not block the game thread.
+	* If the screen reader user is active and no announcements are currently spoken, the announcement will be spoken immediately.
+	* If another announcement is currently being spoken, the passed in announcement could be queued or interrupt the currently spoken announcement.
+	* @param InUserId The user Id of the screen reader user the announcement is intended for 
+	* @param InAnnouncement - The announcement requested to be spoken to the screen reader user
+	* @return FScreenReaderReply::Handled() if the request was successfully processed. Else return FScreenReader::Unhandled()
+	* @see FScreenReaderAnnouncement, FScreenReaderAnnouncementInfo
+	*/
+	UFUNCTION(BlueprintCallable, Category = "SlateScreenReader")
+	FScreenReaderReply RequestSpeak(UPARAM(DisplayName = "Screen Reader User Id") int32 InUserId, UPARAM(DisplayName = "Announcement") FScreenReaderAnnouncement InAnnouncement);
+
+	/**
+	* Requests the information about the accessibility widget a user is focused on to be read out.
+	* If nothing is currently being focused on by the screen reader user, nothing will be read out.
+	* The same guarantees about the announcement being spoken in RequestSpeak() apply for this function.
+	* Nothing will happen if the passed in user Id is not already registered with the screen reader.
+	* @param InUserId The user Id to request its accessible focus to be read out
+	* @return FScreenReaderReply::Handled() if the user's focused widget's information is successfully spoken. Else FScreenReaderReply::Unhandled() is returned. 
+	*/
+	UFUNCTION(BlueprintCallable, Category = "SlateScreenReader")
+	FScreenReaderReply RequestSpeakFocusedWidget(UPARAM(DisplayName = "Screen Reader User Id") int32 InUserId);
+
+	/**
+	* Immediately stops speaking any currently spoken announcement for a particular screen reader user.
+	* Does nothing if there is no announcement currently being spoken for the user or if the user Id is not registered with the screen reader.
+	* @param InUserId The user Id of the screen reader user to request announcements to be stopped.
+	* @return FScreenReader::Handled() if the request to stop speaking is successfully processed. Else FScreenReaderReply::Unhandled() is returned.
+	*/
+	UFUNCTION(BlueprintCallable, Category = "SlateScreenReader")
+	FScreenReaderReply StopSpeaking(UPARAM(DisplayName = "Screen Reader User Id") int32 InUserId);
+
+	/**
+	* Returns true if the screen reader is speaking text to a particular user.
+	* Returns false if no no announcements are being spoken to the user or if the user Id is not registered.
+	* @param InUserId The user Id of the screen reader user to check if any announcements are being spoken to.
+	*/
+	UFUNCTION(BlueprintCallable, Category = "SlateScreenReader")
+	bool IsSpeaking(UPARAM(DisplayName = "Screen Reader User Id") int32 InUserId) const;
+
+	/**
+	* Returns the volume text to speech will be speaking at for a screen reader user.  Value is between 0.0f and 1.0f.
+	* If the provided user Id doesn't exist, 0.0f will be returned.
+	* @param InUserId The user Id of the screen reader user to retrieve the text to speech volume for 
+	*/
+	UFUNCTION(BlueprintCallable, Category = "SlateScreenReader")
+	float GetSpeechVolume(UPARAM(DisplayName = "Screen Reader User Id") int32 InUserId) const;
+
+	/**
+	* Sets the volume text to speech will be speaking at for a screen reader user.  
+	* If the provided user Id doesn't exist, nothing will happen.
+	* @param InUserId The user Id of the screen reader user to set the text to speech volume for
+	* @param InVolume The volume text to speech will be set at for the provided screen reader user. Value will be clamped between 0.0f and 1.0f.
+	* @return FScreenReaderReply::Handled() is returned if the speech volume is successfully set for the screen reader user. Otherwise, FScreenReaderReply::Unhandled() is returned.
+	*/
+	UFUNCTION(BlueprintCallable, Category = "SlateScreenReader")
+	FScreenReaderReply SetSpeechVolume(UPARAM(DisplayName = "Screen Reader User Id") int32 InUserId, UPARAM(DisplayName = "Volume") float InVolume);
+
+	/**
+	* Returns the rate text to speech will be speaking at for a screen reader user.  Value is between 0.0f and 1.0f.
+	* If the provided user Id doesn't exist, 0.0f will be returned.
+	* @param InUserId The user Id of the screen reader user to retrieve the text to speech rate for
 	*/
-	UFUNCTION(BlueprintCallable, Category="ScreenReader")
-	void Deactivate();
+	UFUNCTION(BlueprintCallable, Category = "SlateScreenReader")
+	float GetSpeechRate(UPARAM(DisplayName = "Screen Reader User Id") int32 InUserId) const;
+
+	/**
+	* Sets the rate text to speech will be speaking at for a screen reader user.
+	* If the provided user Id doesn't exist, nothing will happen.
+	* @param InUserId The user Id of the screen reader user to set the text to speech rate for
+	* @param InRate The rate text to speech will be set at for the provided screen reader user. Value will be clamped between 0.0f and 1.0f.
+	* @return FScreenReaderReply::Handled() is returned if the speech rate is successfully set for the screen reader user. Otherwise, FScreenReaderReply::Unhandled() is returned.
+	*/
+	UFUNCTION(BlueprintCallable, Category = "SlateScreenReader")
+	FScreenReaderReply SetSpeechRate(UPARAM(DisplayName = "Screen Reader User Id") int32 InUserId, UPARAM(DisplayName = "Rate") float InRate);
+
+	/**
+	* Mutes the text to speech for a screen reader user.
+	* If the provided user Id doesn't exist, nothing will happen.
+	* @param InUserId The user Id of the screen reader user to mute text to speech for.
+	* @return FScreenReaderReply::Handled() is returned if the screen reader user is successfully muted. Otherwise, FScreenReaderReply::Unhandled() is returned.
+	*/
+	UFUNCTION(BlueprintCallable, Category = "SlateScreenReader")
+	FScreenReaderReply MuteSpeech(UPARAM(DisplayName = "Screen Reader User Id") int32 InUserId);
+
+	/**
+	* Unmutes the text to speech for a screen reader user.
+	* If the provided user Id doesn't exist, nothing will happen.
+	* @param InUserId The user Id of the screen reader user to unmute text to speech for.
+	* @return FScreenReaderReply::Handled() is returned if the screen reader user is successfully unmuted. Otherwise, FScreenReaderReply::Unhandled() is returned.
+	*/
+	UFUNCTION(BlueprintCallable, Category = "SlateScreenReader")
+	FScreenReaderReply UnmuteSpeech(UPARAM(DisplayName = "Screen Reader User Id") int32 InUserId);
+
+	/**
+	* Returns true if text to speech for a screen reader user is muted. Otherwise false is returned.
+	* If the provided user Id doesn't exist, false will be returned.
+	* @param InUserId The screen reader user Id to check if text to speech is muted.
+	*/
+	UFUNCTION(BlueprintCallable, Category = "SlateScreenReader")
+	bool IsSpeechMuted(UPARAM(DisplayName = "Screen Reader User Id") int32 InUserId) const;
+
 	/** Returns the underlying screen reader */
 	TSharedRef<FScreenReaderBase> GetScreenReader() const;
 
