Removal of TEXT used in static_asserts (redundant in ANSI/wide modes, broken in UTF-8 mode).

diff --git a/Engine/Plugins/Enterprise/DatasmithImporter/Source/ExternalSource/Public/SourceUri.h b/Engine/Plugins/Enterprise/DatasmithImporter/Source/ExternalSource/Public/SourceUri.h
--- a/Engine/Plugins/Enterprise/DatasmithImporter/Source/ExternalSource/Public/SourceUri.h
+++ b/Engine/Plugins/Enterprise/DatasmithImporter/Source/ExternalSource/Public/SourceUri.h
@@ -98,7 +98,7 @@ namespace UE::DatasmithImporter
 		template<typename TContainer>
 		static FString BuildQueryString(const TContainer& QueryContainer)
 		{
-			static_assert(std::is_same<typename TContainer::ElementType, TPair<FString, FString>>::value, TEXT("Query container must consist of a collection of TPair<FString, FString>."));
+			static_assert(std::is_same<typename TContainer::ElementType, TPair<FString, FString>>::value, "Query container must consist of a collection of TPair<FString, FString>.");
 
 			if (QueryContainer.Num() > 0)
 			{
diff --git a/Engine/Plugins/Experimental/ScreenReader/Source/ScreenReader/Private/ScreenReaderModule.cpp b/Engine/Plugins/Experimental/ScreenReader/Source/ScreenReader/Private/ScreenReaderModule.cpp
--- a/Engine/Plugins/Experimental/ScreenReader/Source/ScreenReader/Private/ScreenReaderModule.cpp
+++ b/Engine/Plugins/Experimental/ScreenReader/Source/ScreenReader/Private/ScreenReaderModule.cpp
@@ -4,7 +4,7 @@
 #include "ScreenReaderLog.h"
 #include "Modules/ModuleManager.h"
 
-static_assert(WITH_ACCESSIBILITY, TEXT("Trying to use the screen reader plugin with accessibility disabled. Accessibility must be enabled to use this plugin. Either enable accessibility or disable the plugin."));
+static_assert(WITH_ACCESSIBILITY, "Trying to use the screen reader plugin with accessibility disabled. Accessibility must be enabled to use this plugin. Either enable accessibility or disable the plugin.");
 
 IMPLEMENT_MODULE(FDefaultModuleImpl, ScreenReader);
 
diff --git a/Engine/Plugins/Experimental/SlateScreenReader/Source/SlateScreenReader/Public/SlateScreenReaderModule.h b/Engine/Plugins/Experimental/SlateScreenReader/Source/SlateScreenReader/Public/SlateScreenReaderModule.h
--- a/Engine/Plugins/Experimental/SlateScreenReader/Source/SlateScreenReader/Public/SlateScreenReaderModule.h
+++ b/Engine/Plugins/Experimental/SlateScreenReader/Source/SlateScreenReader/Public/SlateScreenReaderModule.h
@@ -5,7 +5,7 @@
 #include "Modules/ModuleManager.h"
 #include "Modules/ModuleInterface.h"
 
-static_assert(WITH_ACCESSIBILITY, TEXT("Trying to use the Slate screen reader plugin with accessibility disabled. Accessibility must be enabled to use this plugin. Either enable accessibility or disable the plugin."));
+static_assert(WITH_ACCESSIBILITY, "Trying to use the Slate screen reader plugin with accessibility disabled. Accessibility must be enabled to use this plugin. Either enable accessibility or disable the plugin.");
 
 class IScreenReaderBuilder;
 
diff --git a/Engine/Source/Runtime/Core/Private/Android/AndroidPlatformProcess.cpp b/Engine/Source/Runtime/Core/Private/Android/AndroidPlatformProcess.cpp
--- a/Engine/Source/Runtime/Core/Private/Android/AndroidPlatformProcess.cpp
+++ b/Engine/Source/Runtime/Core/Private/Android/AndroidPlatformProcess.cpp
@@ -289,7 +289,7 @@ static int AndroidThreadPriorityToNice(EThreadPriority NewPriority)
 		NORMAL+2,	//	TPri_SlightlyBelowNormal
 		URGENT,		//	TPri_TimeCritical
 	};
-	static_assert(UE_ARRAY_COUNT(ThreadPriorityToNiceValue) == EThreadPriority::TPri_Num, TEXT("This array is expected to be 1:1 mapping with EThreadPriority."));
+	static_assert(UE_ARRAY_COUNT(ThreadPriorityToNiceValue) == EThreadPriority::TPri_Num, "This array is expected to be 1:1 mapping with EThreadPriority.");
 	return ThreadPriorityToNiceValue[(int)NewPriority];
 }
 
diff --git a/Engine/Source/Runtime/Core/Private/Unix/UnixPlatformProcess.cpp b/Engine/Source/Runtime/Core/Private/Unix/UnixPlatformProcess.cpp
--- a/Engine/Source/Runtime/Core/Private/Unix/UnixPlatformProcess.cpp
+++ b/Engine/Source/Runtime/Core/Private/Unix/UnixPlatformProcess.cpp
@@ -313,7 +313,7 @@ bool FUnixPlatformProcess::SetProcessLimits(EProcessResource::Type Resource, uin
 {
 	rlimit NativeLimit;
 
-	static_assert(sizeof(long) == sizeof(NativeLimit.rlim_cur), TEXT("Platform has atypical rlimit type."));
+	static_assert(sizeof(long) == sizeof(NativeLimit.rlim_cur), "Platform has atypical rlimit type.");
 
 	// 32-bit platforms set limits as long
 	if (sizeof(NativeLimit.rlim_cur) < sizeof(Limit))
diff --git a/Engine/Source/Runtime/Core/Public/Math/MovingWindowAverageFast.h b/Engine/Source/Runtime/Core/Public/Math/MovingWindowAverageFast.h
--- a/Engine/Source/Runtime/Core/Public/Math/MovingWindowAverageFast.h
+++ b/Engine/Source/Runtime/Core/Public/Math/MovingWindowAverageFast.h
@@ -18,7 +18,7 @@ public:
 		, RemoveNextIdx(0)
 		, NumValuesUsed(0)
 	{
-		static_assert(ArraySize > 0, TEXT("ArraySize must be greater than zero"));
+		static_assert(ArraySize > 0, "ArraySize must be greater than zero");
 	}
 
 	void PushValue(T Value)
diff --git a/Engine/Source/Runtime/MovieScene/Public/EntitySystem/MovieSceneEntitySystemTask.h b/Engine/Source/Runtime/MovieScene/Public/EntitySystem/MovieSceneEntitySystemTask.h
--- a/Engine/Source/Runtime/MovieScene/Public/EntitySystem/MovieSceneEntitySystemTask.h
+++ b/Engine/Source/Runtime/MovieScene/Public/EntitySystem/MovieSceneEntitySystemTask.h
@@ -705,7 +705,7 @@ public:
 	template<int Index>
 	FORCEINLINE auto GetAccessor() const
 	{
-		static_assert(Index < sizeof...(T), TEXT("Invalid component index specified"));
+		static_assert(Index < sizeof...(T), "Invalid component index specified");
 		return Accessors.template Get<Index>();
 	}
 
diff --git a/Engine/Source/Runtime/RenderCore/Public/ShaderCodeArchive.h b/Engine/Source/Runtime/RenderCore/Public/ShaderCodeArchive.h
--- a/Engine/Source/Runtime/RenderCore/Public/ShaderCodeArchive.h
+++ b/Engine/Source/Runtime/RenderCore/Public/ShaderCodeArchive.h
@@ -414,7 +414,7 @@ struct FIoStoreShaderCodeEntry
 	}
 };
 
-static_assert(sizeof(FIoStoreShaderCodeEntry) == sizeof(uint64), TEXT("To reduce memory footprint, shader code entries should be as small as possible"));
+static_assert(sizeof(FIoStoreShaderCodeEntry) == sizeof(uint64), "To reduce memory footprint, shader code entries should be as small as possible");
 
 /** Descriptor of a group of shaders compressed together. This groups already deduplicated, and possibly unrelated, shaders, so this is a distinct concept from a shader map. */
 struct FIoStoreShaderGroupEntry
